[CisFix] Building CmdLink project fails in Visual Studio
- Linked CmdLink against engine so it can resolve globals

diff --git a/Engine/Source/Programs/CmdLink/CmdLink.Target.cs b/Engine/Source/Programs/CmdLink/CmdLink.Target.cs
--- a/Engine/Source/Programs/CmdLink/CmdLink.Target.cs
+++ b/Engine/Source/Programs/CmdLink/CmdLink.Target.cs
@@ -23,7 +23,7 @@ public class CmdLinkTarget : TargetRules
 		bCompileICU = false;
 
 		// Currently this app is not linking against the engine, so we'll compile out references from Core to the rest of the engine
-		bCompileAgainstEngine = false;
+		bCompileAgainstEngine = true;
 		bCompileAgainstCoreUObject = false;
 		bCompileAgainstApplicationCore = false;
 		bBuildWithEditorOnlyData = false;
