Fix compilation of CmdLink (missing external symbols from Core)

diff --git a/Engine/Source/Programs/CmdLink/CmdLink.Build.cs b/Engine/Source/Programs/CmdLink/CmdLink.Build.cs
--- a/Engine/Source/Programs/CmdLink/CmdLink.Build.cs
+++ b/Engine/Source/Programs/CmdLink/CmdLink.Build.cs
@@ -10,6 +10,7 @@ public class CmdLink : ModuleRules
             new string[]
             {
 	            "Core",
+				"Projects",
             }
         );
         
diff --git a/Engine/Source/Programs/CmdLink/Private/CmdLink.cpp b/Engine/Source/Programs/CmdLink/Private/CmdLink.cpp
--- a/Engine/Source/Programs/CmdLink/Private/CmdLink.cpp
+++ b/Engine/Source/Programs/CmdLink/Private/CmdLink.cpp
@@ -5,6 +5,10 @@
 #include "Core.h"
 #include "Windows/WindowsPlatformNamedPipe.h"
 
+#include "RequiredProgramMainCPPInclude.h"
+
+IMPLEMENT_APPLICATION(CmdLink, "CmdLink")
+
 enum EErrorMessage
 {
 	SUCCESS = 0,
